Build UnrealCEFSubProcess for Win32 as part of the Win64 tools node, so it can be included in the Rocket build.

diff --git a/Engine/Source/Programs/UnrealCEFSubProcess/UnrealCEFSubProcess.Target.cs b/Engine/Source/Programs/UnrealCEFSubProcess/UnrealCEFSubProcess.Target.cs
--- a/Engine/Source/Programs/UnrealCEFSubProcess/UnrealCEFSubProcess.Target.cs
+++ b/Engine/Source/Programs/UnrealCEFSubProcess/UnrealCEFSubProcess.Target.cs
@@ -78,6 +78,18 @@ public class UnrealCEFSubProcessTarget : TargetRules
         return true;
     }
 
+    public override List<UnrealTargetPlatform> GUBP_ToolPlatforms(UnrealTargetPlatform InHostPlatform)
+    {
+		if(InHostPlatform == UnrealTargetPlatform.Win64)
+		{
+			return new List<UnrealTargetPlatform>{ UnrealTargetPlatform.Win32, UnrealTargetPlatform.Win64 };
+		}
+		else
+		{
+			return new List<UnrealTargetPlatform> { InHostPlatform };
+		}
+    }
+
     public override bool GUBP_NeedsPlatformSpecificDLLs()
     {
         return true;
