Fix UnrealCEFSubprocess build breakage.

diff --git a/Engine/Source/Programs/UnrealCEFSubProcess/Private/UnrealCEFSubProcessCallbackRegistry.cpp b/Engine/Source/Programs/UnrealCEFSubProcess/Private/UnrealCEFSubProcessCallbackRegistry.cpp
--- a/Engine/Source/Programs/UnrealCEFSubProcess/Private/UnrealCEFSubProcessCallbackRegistry.cpp
+++ b/Engine/Source/Programs/UnrealCEFSubProcess/Private/UnrealCEFSubProcessCallbackRegistry.cpp
@@ -3,6 +3,8 @@
 #include "UnrealCEFSubProcess.h"
 #include "UnrealCEFSubProcessCallbackRegistry.h"
 
+#if WITH_CEF3
+
 FGuid FUnrealCEFSubProcessCallbackRegistry::FindOrAdd(CefRefPtr<CefV8Context> Context, CefRefPtr<CefV8Value> Object, CefRefPtr<CefV8Value> Function, CefRefPtr<CefV8Value> OnError, bool bOneShot)
 {
 	if (! bOneShot)
@@ -37,3 +39,5 @@ void FUnrealCEFSubProcessCallbackRegistry::RemoveByContext(CefRefPtr<CefV8Contex
 		}
 	}
 }
+
+#endif // WITH_CEF3
