[INTEGRATE] Merging CL#2620775 using UE4-To-UE4-LauncherDev:

Fix null pointer error in UnrealCEFSubProcess when passing the same callback handler more than once to a bound UObject method.

diff --git a/Engine/Source/Programs/UnrealCEFSubProcess/Private/UnrealCEFSubProcessCallbackRegistry.cpp b/Engine/Source/Programs/UnrealCEFSubProcess/Private/UnrealCEFSubProcessCallbackRegistry.cpp
--- a/Engine/Source/Programs/UnrealCEFSubProcess/Private/UnrealCEFSubProcessCallbackRegistry.cpp
+++ b/Engine/Source/Programs/UnrealCEFSubProcess/Private/UnrealCEFSubProcessCallbackRegistry.cpp
@@ -15,7 +15,7 @@ FGuid FUnrealCEFSubProcessCallbackRegistry::FindOrAdd(CefRefPtr<CefV8Context> Co
 			if (! Value.bOneShot
 				&& Context->IsSame(Value.Context)
 				&& Function->IsSame(Value.Function)
-				&& Object->IsSame(Value.Object)
+				&& (Object == nullptr && Value.Object == nullptr || Object->IsSame(Value.Object))
 				&& (OnError == nullptr && Value.OnError == nullptr || OnError->IsSame(Value.OnError)))
 			{
 				return Iter.Key();
