Remove hack forcing UnrealCEFSubProcess to build with VS 2013.

diff --git a/Engine/Source/Programs/AutomationTool/GUBP/LegacyNodes.cs b/Engine/Source/Programs/AutomationTool/GUBP/LegacyNodes.cs
--- a/Engine/Source/Programs/AutomationTool/GUBP/LegacyNodes.cs
+++ b/Engine/Source/Programs/AutomationTool/GUBP/LegacyNodes.cs
@@ -639,7 +639,7 @@ partial class GUBP
                     {
 						// Always compile CRC with VS2013 due to requiring the 2015 redist
 						string AddTargetArgs = AddArgs;
-						if(String.Compare(ProgramTarget.TargetName, "CrashReportClient", true) == 0)
+						if(String.Compare(ProgramTarget.TargetName, "CrashReportClient", true) == 0 || String.Compare(ProgramTarget.TargetName, "UnrealCEFSubProcess", true) == 0)
 						{
 							AddTargetArgs += " -2013";
 						}
diff --git a/Engine/Source/Programs/UnrealCEFSubProcess/UnrealCEFSubProcess.Target.cs b/Engine/Source/Programs/UnrealCEFSubProcess/UnrealCEFSubProcess.Target.cs
--- a/Engine/Source/Programs/UnrealCEFSubProcess/UnrealCEFSubProcess.Target.cs
+++ b/Engine/Source/Programs/UnrealCEFSubProcess/UnrealCEFSubProcess.Target.cs
@@ -41,18 +41,6 @@ public class UnrealCEFSubProcessTarget : TargetRules
 			);
 	}
 
-    public override bool ConfigureToolchain(TargetInfo Target)
-    {
-        if (Target.Platform == UnrealTargetPlatform.Win32 || Target.Platform == UnrealTargetPlatform.Win64)
-        {
-            if (Target.Configuration != UnrealTargetConfiguration.Debug)
-            {
-                WindowsPlatform.Compiler = WindowsCompiler.VisualStudio2013;
-            }
-        }
-        return true;
-    }
-
 	public override void SetupGlobalEnvironment(
 		TargetInfo Target,
 		ref LinkEnvironmentConfiguration OutLinkEnvironmentConfiguration,
