Disable shadow variable warnings when compiling UnrealCEFSubProcess with Visual Studio 2015.

diff --git a/Engine/Source/Programs/UnrealCEFSubProcess/UnrealCEFSubProcess.Build.cs b/Engine/Source/Programs/UnrealCEFSubProcess/UnrealCEFSubProcess.Build.cs
--- a/Engine/Source/Programs/UnrealCEFSubProcess/UnrealCEFSubProcess.Build.cs
+++ b/Engine/Source/Programs/UnrealCEFSubProcess/UnrealCEFSubProcess.Build.cs
@@ -29,6 +29,8 @@ public class UnrealCEFSubProcess : ModuleRules
 		AddEngineThirdPartyPrivateStaticDependencies(Target,
 			"CEF3"
 			);
+			
+		bEnableShadowVariableWarnings = false;
 	}
 }
 
