Convert from wchar to tchar and tchar_to_wchar when needed in UnrealCEF

diff --git a/Engine/Source/Programs/UnrealCEFSubProcess/Private/UnrealCEFSubProcessApp.cpp b/Engine/Source/Programs/UnrealCEFSubProcess/Private/UnrealCEFSubProcessApp.cpp
--- a/Engine/Source/Programs/UnrealCEFSubProcess/Private/UnrealCEFSubProcessApp.cpp
+++ b/Engine/Source/Programs/UnrealCEFSubProcess/Private/UnrealCEFSubProcessApp.cpp
@@ -22,7 +22,7 @@ void FUnrealCEFSubProcessApp::OnContextReleased( CefRefPtr<CefBrowser> Browser,
 bool FUnrealCEFSubProcessApp::OnProcessMessageReceived( CefRefPtr<CefBrowser> Browser, CefProcessId SourceProcess, CefRefPtr<CefProcessMessage> Message )
 {
 	bool Result = false;
-	FString MessageName = Message->GetName().ToWString().c_str();
+	FString MessageName = WCHAR_TO_TCHAR(Message->GetName().ToWString().c_str());
 	if (MessageName.StartsWith(TEXT("UE::")))
 	{
 		Result = RemoteScripting.OnProcessMessageReceived(Browser, SourceProcess, Message);
@@ -92,4 +92,4 @@ void FUnrealCEFSubProcessApp::OnFocusedNodeChanged(CefRefPtr<CefBrowser> Browser
 }
 #endif
 
-#endif
\ No newline at end of file
+#endif
diff --git a/Engine/Source/Programs/UnrealCEFSubProcess/Private/UnrealCEFSubProcessRemoteMethodHandler.cpp b/Engine/Source/Programs/UnrealCEFSubProcess/Private/UnrealCEFSubProcessRemoteMethodHandler.cpp
--- a/Engine/Source/Programs/UnrealCEFSubProcess/Private/UnrealCEFSubProcessRemoteMethodHandler.cpp
+++ b/Engine/Source/Programs/UnrealCEFSubProcess/Private/UnrealCEFSubProcessRemoteMethodHandler.cpp
@@ -59,9 +59,9 @@ bool FUnrealCEFSubProcessRemoteObject::ExecuteMethod(const CefString& MethodName
 	FGuid CallbackGuid = RemoteScripting->CallbackRegistry.FindOrAdd(Context, Retval, Accept, Reject, true);
 	CefRefPtr<CefProcessMessage> Message = CefProcessMessage::Create("UE::ExecuteUObjectMethod");
 	CefRefPtr<CefListValue> MessageArguments = Message->GetArgumentList();
-	MessageArguments->SetString(0, CefString(*ObjectId.ToString(EGuidFormats::Digits)));
+	MessageArguments->SetString(0, CefString(TCHAR_TO_WCHAR(*ObjectId.ToString(EGuidFormats::Digits))));
 	MessageArguments->SetString(1, MethodName);
-	MessageArguments->SetString(2, CefString(*CallbackGuid.ToString(EGuidFormats::Digits)));
+	MessageArguments->SetString(2, CefString(TCHAR_TO_WCHAR(*CallbackGuid.ToString(EGuidFormats::Digits))));
 	MessageArguments->SetList(3, RemoteScripting->V8ArrayToCef(Arguments));
 
 	Browser->SendProcessMessage(PID_BROWSER, Message);
@@ -74,8 +74,8 @@ FUnrealCEFSubProcessRemoteObject::~FUnrealCEFSubProcessRemoteObject()
 {
 	CefRefPtr<CefProcessMessage> Message = CefProcessMessage::Create("UE::ReleaseUObject");
 	CefRefPtr<CefListValue> MessageArguments = Message->GetArgumentList();
-	MessageArguments->SetString(0, CefString(*ObjectId.ToString(EGuidFormats::Digits)));
+	MessageArguments->SetString(0, CefString(TCHAR_TO_WCHAR(*ObjectId.ToString(EGuidFormats::Digits))));
 	Browser->SendProcessMessage(PID_BROWSER, Message);
 }
 
-#endif
\ No newline at end of file
+#endif
diff --git a/Engine/Source/Programs/UnrealCEFSubProcess/Private/UnrealCEFSubProcessRemoteScripting.cpp b/Engine/Source/Programs/UnrealCEFSubProcess/Private/UnrealCEFSubProcessRemoteScripting.cpp
--- a/Engine/Source/Programs/UnrealCEFSubProcess/Private/UnrealCEFSubProcessRemoteScripting.cpp
+++ b/Engine/Source/Programs/UnrealCEFSubProcess/Private/UnrealCEFSubProcessRemoteScripting.cpp
@@ -10,7 +10,7 @@ CefRefPtr<CefV8Value> FUnrealCEFSubProcessRemoteScripting::CefToV8(CefRefPtr<Cef
 	// Custom types are encoded inside dictionary values with a $type and a $value property
 	if ( Dictionary->GetType("$type") == VTYPE_STRING)
 	{
-		FString Type = Dictionary->GetString("$type").ToWString().c_str();
+		FString Type = WCHAR_TO_TCHAR(Dictionary->GetString("$type").ToWString().c_str());
 		if ( Type == "struct" &&  Dictionary->GetType("$value") == VTYPE_DICTIONARY)
 		{
 			return CefToPlainV8Object(Dictionary->GetDictionary("$value"));
@@ -18,7 +18,7 @@ CefRefPtr<CefV8Value> FUnrealCEFSubProcessRemoteScripting::CefToV8(CefRefPtr<Cef
 		if ( Type == "uobject" && Dictionary->GetType("$id") == VTYPE_STRING && Dictionary->GetType("$methods") == VTYPE_LIST)
 		{
 			FGuid Guid;
-			if (FGuid::Parse(Dictionary->GetString("$id").ToWString().c_str(), Guid))
+			if (FGuid::Parse(WCHAR_TO_TCHAR(Dictionary->GetString("$id").ToWString().c_str()), Guid))
 			{
 				CefRefPtr<CefListValue> Methods = Dictionary->GetList("$methods");
 				return CreateUObjectProxy(Guid, Methods);
@@ -45,7 +45,7 @@ CefRefPtr<CefV8Value> FUnrealCEFSubProcessRemoteScripting::CreateUObjectProxy(FG
 		CefRefPtr<CefV8Value> FunctionProxy = CefV8Value::CreateFunction(MethodName, new FUnrealCEFSubProcessRemoteMethodHandler(Remote, MethodName));
 		Result->SetValue(MethodName, FunctionProxy, static_cast<CefV8Value::PropertyAttribute>(V8_PROPERTY_ATTRIBUTE_DONTDELETE | V8_PROPERTY_ATTRIBUTE_READONLY ));
 	}
-	Result->SetValue("$id", CefV8Value::CreateString(*ObjectId.ToString(EGuidFormats::Digits)),  static_cast<CefV8Value::PropertyAttribute>(V8_PROPERTY_ATTRIBUTE_DONTDELETE | V8_PROPERTY_ATTRIBUTE_READONLY | V8_PROPERTY_ATTRIBUTE_DONTENUM));
+	Result->SetValue("$id", CefV8Value::CreateString(TCHAR_TO_WCHAR(*ObjectId.ToString(EGuidFormats::Digits))),  static_cast<CefV8Value::PropertyAttribute>(V8_PROPERTY_ATTRIBUTE_DONTDELETE | V8_PROPERTY_ATTRIBUTE_READONLY | V8_PROPERTY_ATTRIBUTE_DONTENUM));
 	return Result;
 }
 
@@ -128,7 +128,7 @@ CefRefPtr<CefDictionaryValue> FUnrealCEFSubProcessRemoteScripting::V8FunctionToC
 	CefRefPtr<CefDictionaryValue> Result = CefDictionaryValue::Create();
 	FGuid Guid = CallbackRegistry.FindOrAdd(CefV8Context::GetCurrentContext(), Object, Function);
 	Result->SetString("$type", "callback");
-	Result->SetString("$id", *Guid.ToString(EGuidFormats::Digits));
+	Result->SetString("$id", TCHAR_TO_WCHAR(*Guid.ToString(EGuidFormats::Digits)));
 	Result->SetString("$name", Function->GetFunctionName());
 	return Result;
 }
@@ -147,7 +147,7 @@ bool FUnrealCEFSubProcessRemoteScripting::HandleExecuteJSFunctionMessage(CefRefP
 		return false;
 	}
 
-	if (!FGuid::Parse(FString(MessageArguments->GetString(0).ToWString().c_str()), Guid))
+	if (!FGuid::Parse(FString(WCHAR_TO_TCHAR(MessageArguments->GetString(0).ToWString().c_str())), Guid))
 	{
 		// Invalid GUID
 		return false;
@@ -277,7 +277,7 @@ bool FUnrealCEFSubProcessRemoteScripting::HandleDeleteValueMessage(CefRefPtr<Cef
 		bool bPermanent = Argument->GetBool("permanent");
 
 		FGuid Guid;
-		if (!FGuid::Parse(Id.ToWString().c_str(), Guid))
+		if (!FGuid::Parse(WCHAR_TO_TCHAR(Id.ToWString().c_str()), Guid))
 		{
 			return false;
 		}
@@ -329,7 +329,7 @@ bool FUnrealCEFSubProcessRemoteScripting::OnProcessMessageReceived(CefRefPtr<Cef
     CefRefPtr<CefProcessMessage> Message)
 {
 	bool Result = false;
-	FString MessageName = Message->GetName().ToWString().c_str();
+	FString MessageName = WCHAR_TO_TCHAR(Message->GetName().ToWString().c_str());
 	if (MessageName == TEXT("UE::ExecuteJSFunction"))
 	{
 		Result = HandleExecuteJSFunctionMessage(Message->GetArgumentList());
@@ -388,4 +388,4 @@ void FUnrealCEFSubProcessRemoteScripting::InitPermanentBindings(int32 BrowserID,
 }
 
 
-#endif
\ No newline at end of file
+#endif
