Fix error attempting to embed manifest into executable.

diff --git a/Engine/Source/Programs/UnrealCEFSubProcess/UnrealCEFSubProcess.Target.cs b/Engine/Source/Programs/UnrealCEFSubProcess/UnrealCEFSubProcess.Target.cs
--- a/Engine/Source/Programs/UnrealCEFSubProcess/UnrealCEFSubProcess.Target.cs
+++ b/Engine/Source/Programs/UnrealCEFSubProcess/UnrealCEFSubProcess.Target.cs
@@ -41,5 +41,8 @@ public class UnrealCEFSubProcessTarget : TargetRules
 
 		// Epic Games Launcher needs to run on OS X 10.9, so CEFSubProcess needs this as well
 		bEnableOSX109Support = true;
+
+		// Already a manifest specified through resource file
+		WindowsPlatform.ManifestFile = null;
 	}
 }
